Fix IWYU issues for LocationServices plugins
#ue4
#android
#ios

diff --git a/Engine/Plugins/Runtime/LocationServicesAndroidImpl/Source/LocationServicesAndroidEditor/Private/LocationServicesAndroidEditorPrivatePCH.h b/Engine/Plugins/Runtime/LocationServicesAndroidImpl/Source/LocationServicesAndroidEditor/Private/LocationServicesAndroidEditorPrivatePCH.h
deleted file mode 100644
diff --git a/Engine/Plugins/Runtime/LocationServicesAndroidImpl/Source/LocationServicesAndroidImpl/Private/LocationServicesAndroidImpl.cpp b/Engine/Plugins/Runtime/LocationServicesAndroidImpl/Source/LocationServicesAndroidImpl/Private/LocationServicesAndroidImpl.cpp
--- a/Engine/Plugins/Runtime/LocationServicesAndroidImpl/Source/LocationServicesAndroidImpl/Private/LocationServicesAndroidImpl.cpp
+++ b/Engine/Plugins/Runtime/LocationServicesAndroidImpl/Source/LocationServicesAndroidImpl/Private/LocationServicesAndroidImpl.cpp
@@ -1,6 +1,6 @@
 // Copyright 1998-2016 Epic Games, Inc. All Rights Reserved.
 
-#include "LocationServicesAndroidImplPCH.h"
+#include "LocationServicesAndroidImpl.h"
 #include "Android/AndroidJNI.h"
 #include "Android/AndroidApplication.h"
 
diff --git a/Engine/Plugins/Runtime/LocationServicesAndroidImpl/Source/LocationServicesAndroidImpl/Private/LocationServicesAndroidImplModule.cpp b/Engine/Plugins/Runtime/LocationServicesAndroidImpl/Source/LocationServicesAndroidImpl/Private/LocationServicesAndroidImplModule.cpp
--- a/Engine/Plugins/Runtime/LocationServicesAndroidImpl/Source/LocationServicesAndroidImpl/Private/LocationServicesAndroidImplModule.cpp
+++ b/Engine/Plugins/Runtime/LocationServicesAndroidImpl/Source/LocationServicesAndroidImpl/Private/LocationServicesAndroidImplModule.cpp
@@ -1,7 +1,6 @@
 // Copyright 1998-2016 Epic Games, Inc. All Rights Reserved.
 
-#include "LocationServicesAndroidImplPCH.h"
-
+#include "LocationServicesAndroidImplModule.h"
 #include "Classes/LocationServicesBPLibrary.h"
 
 IMPLEMENT_MODULE(FLocationServicesAndroidImplModule, LocationServicesAndroidImpl)
diff --git a/Engine/Plugins/Runtime/LocationServicesAndroidImpl/Source/LocationServicesAndroidImpl/Private/LocationServicesAndroidImplPCH.h b/Engine/Plugins/Runtime/LocationServicesAndroidImpl/Source/LocationServicesAndroidImpl/Private/LocationServicesAndroidImplPCH.h
deleted file mode 100644
diff --git a/Engine/Plugins/Runtime/LocationServicesAndroidImpl/Source/LocationServicesAndroidImpl/Public/LocationServicesAndroidImplModule.h b/Engine/Plugins/Runtime/LocationServicesAndroidImpl/Source/LocationServicesAndroidImpl/Public/LocationServicesAndroidImplModule.h
--- a/Engine/Plugins/Runtime/LocationServicesAndroidImpl/Source/LocationServicesAndroidImpl/Public/LocationServicesAndroidImplModule.h
+++ b/Engine/Plugins/Runtime/LocationServicesAndroidImpl/Source/LocationServicesAndroidImpl/Public/LocationServicesAndroidImplModule.h
@@ -2,7 +2,9 @@
 
 #pragma once
 
-#include "Core.h"
+#include "CoreMinimal.h"
+#include "ModuleInterface.h"
+#include "ModuleManager.h"
 
 class ULocationServicesAndroidImpl;
 /**
diff --git a/Engine/Plugins/Runtime/LocationServicesIOSImpl/Source/LocationServicesIOSEditor/Private/LocationServicesIOSEditorPrivatePCH.h b/Engine/Plugins/Runtime/LocationServicesIOSImpl/Source/LocationServicesIOSEditor/Private/LocationServicesIOSEditorPrivatePCH.h
deleted file mode 100644
diff --git a/Engine/Plugins/Runtime/LocationServicesIOSImpl/Source/LocationServicesIOSImpl/Private/LocationServicesIOSImpl.cpp b/Engine/Plugins/Runtime/LocationServicesIOSImpl/Source/LocationServicesIOSImpl/Private/LocationServicesIOSImpl.cpp
--- a/Engine/Plugins/Runtime/LocationServicesIOSImpl/Source/LocationServicesIOSImpl/Private/LocationServicesIOSImpl.cpp
+++ b/Engine/Plugins/Runtime/LocationServicesIOSImpl/Source/LocationServicesIOSImpl/Private/LocationServicesIOSImpl.cpp
@@ -1,6 +1,6 @@
 // Copyright 1998-2016 Epic Games, Inc. All Rights Reserved.
 
-#include "LocationServicesIOSImplPCH.h"
+#include "LocationServicesIOSImpl.h"
 #include <CoreLocation/CoreLocation.h>
 
 DEFINE_LOG_CATEGORY(LogLocationServicesIOS);
diff --git a/Engine/Plugins/Runtime/LocationServicesIOSImpl/Source/LocationServicesIOSImpl/Private/LocationServicesIOSImplModule.cpp b/Engine/Plugins/Runtime/LocationServicesIOSImpl/Source/LocationServicesIOSImpl/Private/LocationServicesIOSImplModule.cpp
--- a/Engine/Plugins/Runtime/LocationServicesIOSImpl/Source/LocationServicesIOSImpl/Private/LocationServicesIOSImplModule.cpp
+++ b/Engine/Plugins/Runtime/LocationServicesIOSImpl/Source/LocationServicesIOSImpl/Private/LocationServicesIOSImplModule.cpp
@@ -1,7 +1,6 @@
 // Copyright 1998-2016 Epic Games, Inc. All Rights Reserved.
 
-#include "LocationServicesIOSImplPCH.h"
-
+#include "LocationServicesIOSImplModule.h"
 #include "Classes/LocationServicesBPLibrary.h"
 
 IMPLEMENT_MODULE(FLocationServicesIOSImplModule, LocationServicesIOSImpl)
diff --git a/Engine/Plugins/Runtime/LocationServicesIOSImpl/Source/LocationServicesIOSImpl/Private/LocationServicesIOSImplPCH.h b/Engine/Plugins/Runtime/LocationServicesIOSImpl/Source/LocationServicesIOSImpl/Private/LocationServicesIOSImplPCH.h
deleted file mode 100644
diff --git a/Engine/Plugins/Runtime/LocationServicesIOSImpl/Source/LocationServicesIOSImpl/Public/LocationServicesIOSImplModule.h b/Engine/Plugins/Runtime/LocationServicesIOSImpl/Source/LocationServicesIOSImpl/Public/LocationServicesIOSImplModule.h
--- a/Engine/Plugins/Runtime/LocationServicesIOSImpl/Source/LocationServicesIOSImpl/Public/LocationServicesIOSImplModule.h
+++ b/Engine/Plugins/Runtime/LocationServicesIOSImpl/Source/LocationServicesIOSImpl/Public/LocationServicesIOSImplModule.h
@@ -2,7 +2,9 @@
 
 #pragma once
 
-#include "Core.h"
+#include "CoreMinimal.h"
+#include "ModuleInterface.h"
+#include "ModuleManager.h"
 
 class ULocationServicesIOSImpl;
 /**
diff --git a/Engine/Plugins/Runtime/OptionalMobileFeaturesBPLibrary/Source/OptionalMobileFeaturesBPLibrary/Private/OptionalMobileFeaturesBPLibraryPCH.h b/Engine/Plugins/Runtime/OptionalMobileFeaturesBPLibrary/Source/OptionalMobileFeaturesBPLibrary/Private/OptionalMobileFeaturesBPLibraryPCH.h
deleted file mode 100644
