Fix non-unity issues

diff --git a/Engine/Plugins/Editor/PluginBrowser/Source/PluginBrowser/Private/DefaultPluginWizardDefinition.cpp b/Engine/Plugins/Editor/PluginBrowser/Source/PluginBrowser/Private/DefaultPluginWizardDefinition.cpp
--- a/Engine/Plugins/Editor/PluginBrowser/Source/PluginBrowser/Private/DefaultPluginWizardDefinition.cpp
+++ b/Engine/Plugins/Editor/PluginBrowser/Source/PluginBrowser/Private/DefaultPluginWizardDefinition.cpp
@@ -4,6 +4,8 @@
 #include "PluginStyle.h"
 #include "Interfaces/IPluginManager.h"
 #include "HAL/PlatformFilemanager.h"
+#include "App.h"
+#include "GenericPlatformFile.h"
 
 #define LOCTEXT_NAMESPACE "NewPluginWizard"
 
diff --git a/Engine/Plugins/Editor/PluginBrowser/Source/PluginBrowser/Private/SNewPluginWizard.h b/Engine/Plugins/Editor/PluginBrowser/Source/PluginBrowser/Private/SNewPluginWizard.h
--- a/Engine/Plugins/Editor/PluginBrowser/Source/PluginBrowser/Private/SNewPluginWizard.h
+++ b/Engine/Plugins/Editor/PluginBrowser/Source/PluginBrowser/Private/SNewPluginWizard.h
@@ -16,6 +16,8 @@ class SFilePathBlock;
 class STableViewBase;
 enum class ECheckBoxState : uint8;
 
+template <typename ItemType>
+class SListView;
 
 DECLARE_LOG_CATEGORY_EXTERN(LogPluginWizard, Log, All);
 
diff --git a/Engine/Plugins/Editor/PluginBrowser/Source/PluginBrowser/Public/IPluginWizardDefinition.h b/Engine/Plugins/Editor/PluginBrowser/Source/PluginBrowser/Public/IPluginWizardDefinition.h
--- a/Engine/Plugins/Editor/PluginBrowser/Source/PluginBrowser/Public/IPluginWizardDefinition.h
+++ b/Engine/Plugins/Editor/PluginBrowser/Source/PluginBrowser/Public/IPluginWizardDefinition.h
@@ -3,6 +3,9 @@
 #pragma once
 
 #include "ModuleDescriptor.h"
+#include "SlateEnums.h"
+#include "SlateTypes.h"
+#include "ITypedTableView.h"
 
 /**
  * Description of a plugin template
@@ -19,7 +22,7 @@ struct FPluginTemplateDescription
 	FString OnDiskPath;
 
 	/** Brush resource for the image that is dynamically loaded */
-	TSharedPtr< FSlateDynamicImageBrush > PluginIconDynamicImageBrush;
+	TSharedPtr< struct FSlateDynamicImageBrush > PluginIconDynamicImageBrush;
 
 	/** Can the plugin contain content? */
 	bool bCanContainContent;
diff --git a/Engine/Plugins/Runtime/LocationServicesAndroidImpl/Source/LocationServicesAndroidEditor/Classes/LocationServicesAndroidSettings.h b/Engine/Plugins/Runtime/LocationServicesAndroidImpl/Source/LocationServicesAndroidEditor/Classes/LocationServicesAndroidSettings.h
--- a/Engine/Plugins/Runtime/LocationServicesAndroidImpl/Source/LocationServicesAndroidEditor/Classes/LocationServicesAndroidSettings.h
+++ b/Engine/Plugins/Runtime/LocationServicesAndroidImpl/Source/LocationServicesAndroidEditor/Classes/LocationServicesAndroidSettings.h
@@ -2,6 +2,7 @@
 
 #pragma once
 
+#include "UObject/Object.h"
 #include "LocationServicesAndroidSettings.generated.h"
 
 UCLASS(config = Engine, defaultconfig)
diff --git a/Engine/Plugins/Runtime/LocationServicesAndroidImpl/Source/LocationServicesAndroidEditor/Private/LocationServicesAndroidEditor.cpp b/Engine/Plugins/Runtime/LocationServicesAndroidImpl/Source/LocationServicesAndroidEditor/Private/LocationServicesAndroidEditor.cpp
--- a/Engine/Plugins/Runtime/LocationServicesAndroidImpl/Source/LocationServicesAndroidEditor/Private/LocationServicesAndroidEditor.cpp
+++ b/Engine/Plugins/Runtime/LocationServicesAndroidImpl/Source/LocationServicesAndroidEditor/Private/LocationServicesAndroidEditor.cpp
@@ -2,11 +2,16 @@
 
 #include "LocationServicesAndroidEditor.h"
 #include "LocationServicesAndroidSettings.h"
+#include "ModuleManager.h"
+#include "UObject/Class.h"
+#include "UObject/WeakObjectPtr.h"
+
 #if WITH_EDITOR
 #include "ISettingsModule.h"
 #include "ISettingsSection.h"
 #endif
 
+
 IMPLEMENT_MODULE( FLocationServicesAndroidEditorModule, LocationServicesAndroidEditorModule );
 
 #define LOCTEXT_NAMESPACE "FLocationServicesAndroidEditorModule"
diff --git a/Engine/Plugins/Runtime/LocationServicesIOSImpl/Source/LocationServicesIOSEditor/Classes/LocationServicesIOSSettings.h b/Engine/Plugins/Runtime/LocationServicesIOSImpl/Source/LocationServicesIOSEditor/Classes/LocationServicesIOSSettings.h
--- a/Engine/Plugins/Runtime/LocationServicesIOSImpl/Source/LocationServicesIOSEditor/Classes/LocationServicesIOSSettings.h
+++ b/Engine/Plugins/Runtime/LocationServicesIOSImpl/Source/LocationServicesIOSEditor/Classes/LocationServicesIOSSettings.h
@@ -2,6 +2,7 @@
 
 #pragma once
 
+#include "UObject/Object.h"
 #include "LocationServicesIOSSettings.generated.h"
 
 UCLASS(config = Engine, defaultconfig)
diff --git a/Engine/Plugins/Runtime/LocationServicesIOSImpl/Source/LocationServicesIOSEditor/Private/LocationServicesIOSEditor.cpp b/Engine/Plugins/Runtime/LocationServicesIOSImpl/Source/LocationServicesIOSEditor/Private/LocationServicesIOSEditor.cpp
--- a/Engine/Plugins/Runtime/LocationServicesIOSImpl/Source/LocationServicesIOSEditor/Private/LocationServicesIOSEditor.cpp
+++ b/Engine/Plugins/Runtime/LocationServicesIOSImpl/Source/LocationServicesIOSEditor/Private/LocationServicesIOSEditor.cpp
@@ -2,11 +2,16 @@
 
 #include "LocationServicesIOSEditor.h"
 #include "LocationServicesIOSSettings.h"
+#include "ModuleManager.h"
+#include "UObject/Class.h"
+#include "UObject/WeakObjectPtr.h"
+
 #if WITH_EDITOR
 #include "ISettingsModule.h"
 #include "ISettingsSection.h"
 #endif
 
+
 IMPLEMENT_MODULE( FLocationServicesIOSEditorModule, LocationServicesIOSEditorModule );
 
 #define LOCTEXT_NAMESPACE "FLocationServicesIOSEditorModule"
diff --git a/Engine/Source/Developer/Android/AndroidTargetPlatform/Private/AndroidTargetDeviceOutput.h b/Engine/Source/Developer/Android/AndroidTargetPlatform/Private/AndroidTargetDeviceOutput.h
--- a/Engine/Source/Developer/Android/AndroidTargetPlatform/Private/AndroidTargetDeviceOutput.h
+++ b/Engine/Source/Developer/Android/AndroidTargetPlatform/Private/AndroidTargetDeviceOutput.h
@@ -8,6 +8,7 @@
 #include "HAL/PlatformProcess.h"
 #include "Misc/ConfigCacheIni.h"
 #include "Interfaces/ITargetDeviceOutput.h"
+#include "ThreadSafeCounter.h"
 
 class FAndroidTargetDevice;
 
diff --git a/Engine/Source/Editor/AudioEditor/Classes/SoundSubmixGraph/SoundSubmixGraphNode.h b/Engine/Source/Editor/AudioEditor/Classes/SoundSubmixGraph/SoundSubmixGraphNode.h
--- a/Engine/Source/Editor/AudioEditor/Classes/SoundSubmixGraph/SoundSubmixGraphNode.h
+++ b/Engine/Source/Editor/AudioEditor/Classes/SoundSubmixGraph/SoundSubmixGraphNode.h
@@ -1,8 +1,14 @@
 // Copyright 1998-2017 Epic Games, Inc. All Rights Reserved.
 
 #pragma once
+
+#include "CoreMinimal.h"
+#include "UObject/ObjectMacros.h"
+#include "EdGraph/EdGraphNode.h"
 #include "SoundSubmixGraphNode.generated.h"
 
+class USoundSubmix;
+
 UCLASS(MinimalAPI)
 class USoundSubmixGraphNode : public UEdGraphNode
 {
diff --git a/Engine/Source/Editor/AudioEditor/Private/AudioEditorModule.cpp b/Engine/Source/Editor/AudioEditor/Private/AudioEditorModule.cpp
--- a/Engine/Source/Editor/AudioEditor/Private/AudioEditorModule.cpp
+++ b/Engine/Source/Editor/AudioEditor/Private/AudioEditorModule.cpp
@@ -32,6 +32,7 @@
 #include "AssetTypeActions/AssetTypeActions_SoundSubmix.h"
 #include "AssetTypeActions/AssetTypeActions_SoundEffectPreset.h"
 #include "Utils.h"
+#include "UObject/UObjectIterator.h"
 
 const FName AudioEditorAppIdentifier = FName(TEXT("AudioEditorApp"));
 
diff --git a/Engine/Source/Editor/AudioEditor/Private/SSoundSubmixActionMenu.cpp b/Engine/Source/Editor/AudioEditor/Private/SSoundSubmixActionMenu.cpp
--- a/Engine/Source/Editor/AudioEditor/Private/SSoundSubmixActionMenu.cpp
+++ b/Engine/Source/Editor/AudioEditor/Private/SSoundSubmixActionMenu.cpp
@@ -3,6 +3,10 @@
 #include "SSoundSubmixActionMenu.h"
 #include "SoundSubmixGraph/SoundSubmixGraphSchema.h"
 #include "GraphEditor.h"
+#include "Widgets/Input/SEditableTextBox.h"
+#include "EdGraph/EdGraph.h"
+#include "EditorStyleSet.h"
+#include "SlateApplication.h"
 
 #define LOCTEXT_NAMESPACE "SSoundSubmixActionMenu"
 
diff --git a/Engine/Source/Editor/AudioEditor/Private/SoundClassEditor.cpp b/Engine/Source/Editor/AudioEditor/Private/SoundClassEditor.cpp
--- a/Engine/Source/Editor/AudioEditor/Private/SoundClassEditor.cpp
+++ b/Engine/Source/Editor/AudioEditor/Private/SoundClassEditor.cpp
@@ -26,6 +26,7 @@
 #include "IAssetTools.h"
 #include "Factories/SoundClassFactory.h"
 #include "UObject/Package.h"
+#include "UObjectIterator.h"
 
 #define LOCTEXT_NAMESPACE "SoundClassEditor"
 DEFINE_LOG_CATEGORY_STATIC( LogSoundClassEditor, Log, All );
diff --git a/Engine/Source/Editor/AudioEditor/Private/SoundSubmixEditor.cpp b/Engine/Source/Editor/AudioEditor/Private/SoundSubmixEditor.cpp
--- a/Engine/Source/Editor/AudioEditor/Private/SoundSubmixEditor.cpp
+++ b/Engine/Source/Editor/AudioEditor/Private/SoundSubmixEditor.cpp
@@ -20,6 +20,7 @@
 #include "SoundSubmixGraph/SoundSubmixGraphSchema.h"
 #include "AudioEditorModule.h"
 #include "SSoundSubmixActionMenu.h"
+#include "Editor.h"
 
 #define LOCTEXT_NAMESPACE "SoundSubmixEditor"
 DEFINE_LOG_CATEGORY_STATIC( LogSoundSubmixEditor, Log, All );
diff --git a/Engine/Source/Editor/AudioEditor/Private/SoundSubmixEditor.h b/Engine/Source/Editor/AudioEditor/Private/SoundSubmixEditor.h
--- a/Engine/Source/Editor/AudioEditor/Private/SoundSubmixEditor.h
+++ b/Engine/Source/Editor/AudioEditor/Private/SoundSubmixEditor.h
@@ -9,6 +9,8 @@
 //////////////////////////////////////////////////////////////////////////
 // FSoundSubmixEditor
 
+class USoundSubmix;
+
 class FSoundSubmixEditor :	public ISoundSubmixEditor, 
 							public FGCObject, 
 							public FEditorUndoClient
diff --git a/Engine/Source/Editor/AudioEditor/Private/SoundSubmixGraphSchema.cpp b/Engine/Source/Editor/AudioEditor/Private/SoundSubmixGraphSchema.cpp
--- a/Engine/Source/Editor/AudioEditor/Private/SoundSubmixGraphSchema.cpp
+++ b/Engine/Source/Editor/AudioEditor/Private/SoundSubmixGraphSchema.cpp
@@ -9,6 +9,7 @@
 #include "SoundSubmixEditorUtilities.h"
 #include "Sound/SoundSubmix.h"
 #include "GenericCommands.h"
+#include "MultiBoxBuilder.h"
 
 #define LOCTEXT_NAMESPACE "SoundSubmixSchema"
 
diff --git a/Engine/Source/Editor/AudioEditor/Public/ISoundSubmixEditor.h b/Engine/Source/Editor/AudioEditor/Public/ISoundSubmixEditor.h
--- a/Engine/Source/Editor/AudioEditor/Public/ISoundSubmixEditor.h
+++ b/Engine/Source/Editor/AudioEditor/Public/ISoundSubmixEditor.h
@@ -3,6 +3,7 @@
 #pragma once
 
 #include "Toolkits/IToolkitHost.h"
+#include "Toolkits/AssetEditorToolkit.h"
 
 class UEdGraphPin;
 
diff --git a/Engine/Source/Editor/AudioEditor/Public/SoundSubmixEditorUtilities.h b/Engine/Source/Editor/AudioEditor/Public/SoundSubmixEditorUtilities.h
--- a/Engine/Source/Editor/AudioEditor/Public/SoundSubmixEditorUtilities.h
+++ b/Engine/Source/Editor/AudioEditor/Public/SoundSubmixEditorUtilities.h
@@ -2,6 +2,8 @@
 
 #pragma once
 
+#include "CoreMinimal.h"
+
 class UEdGraph;
 class UEdGraphPin;
 
diff --git a/Engine/Source/Editor/ClassViewer/Private/SClassViewer.h b/Engine/Source/Editor/ClassViewer/Private/SClassViewer.h
--- a/Engine/Source/Editor/ClassViewer/Private/SClassViewer.h
+++ b/Engine/Source/Editor/ClassViewer/Private/SClassViewer.h
@@ -13,6 +13,8 @@
 #include "Widgets/Views/STableRow.h"
 #include "Widgets/Views/STreeView.h"
 #include "Settings/ClassViewerSettings.h"
+#include "Engine/EngineTypes.h"
+#include "Misc/StringClassReference.h"
 
 class FClassViewerNode;
 class FMenuBuilder;
diff --git a/Engine/Source/Editor/ClassViewer/Public/ClassViewerProjectSettings.h b/Engine/Source/Editor/ClassViewer/Public/ClassViewerProjectSettings.h
--- a/Engine/Source/Editor/ClassViewer/Public/ClassViewerProjectSettings.h
+++ b/Engine/Source/Editor/ClassViewer/Public/ClassViewerProjectSettings.h
@@ -3,8 +3,13 @@
 #pragma once
 
 #include "CoreMinimal.h"
+#include "UObject/ObjectMacros.h"
+#include "UObject/Object.h"
+#include "Engine/EngineTypes.h"
+#include "Misc/StringClassReference.h"
 #include "ClassViewerProjectSettings.generated.h"
 
+
 /**
  * Implements the settings for the Class Viewer Project Settings
  */
diff --git a/Engine/Source/Editor/ComponentVisualizers/Public/AttenuatedComponentVisualizer.h b/Engine/Source/Editor/ComponentVisualizers/Public/AttenuatedComponentVisualizer.h
--- a/Engine/Source/Editor/ComponentVisualizers/Public/AttenuatedComponentVisualizer.h
+++ b/Engine/Source/Editor/ComponentVisualizers/Public/AttenuatedComponentVisualizer.h
@@ -4,6 +4,8 @@
 
 #include "ComponentVisualizer.h"
 #include "ShowFlags.h"
+#include "SceneView.h"
+#include "SceneManagement.h"
 
 template<class T>
 class COMPONENTVISUALIZERS_API TAttenuatedComponentVisualizer : public FComponentVisualizer
diff --git a/Engine/Source/Editor/DetailCustomizations/Private/ColorGradingVectorCustomization.cpp b/Engine/Source/Editor/DetailCustomizations/Private/ColorGradingVectorCustomization.cpp
--- a/Engine/Source/Editor/DetailCustomizations/Private/ColorGradingVectorCustomization.cpp
+++ b/Engine/Source/Editor/DetailCustomizations/Private/ColorGradingVectorCustomization.cpp
@@ -6,6 +6,8 @@
 #include "IDetailChildrenBuilder.h"
 #include "DetailLayoutBuilder.h"
 #include "DetailWidgetRow.h"
+#include "UObject/UnrealType.h"
+#include "Widgets/Layout/SBox.h"
 
 FColorGradingVectorCustomization::FColorGradingVectorCustomization()
 {
diff --git a/Engine/Source/Editor/DetailCustomizations/Private/DirectoryPathStructCustomization.cpp b/Engine/Source/Editor/DetailCustomizations/Private/DirectoryPathStructCustomization.cpp
--- a/Engine/Source/Editor/DetailCustomizations/Private/DirectoryPathStructCustomization.cpp
+++ b/Engine/Source/Editor/DetailCustomizations/Private/DirectoryPathStructCustomization.cpp
@@ -15,6 +15,7 @@
 #include "DesktopPlatformModule.h"
 #include "IContentBrowserSingleton.h"
 #include "ContentBrowserModule.h"
+#include "MultiBoxBuilder.h"
 
 #define LOCTEXT_NAMESPACE "DirectoryPathStructCustomization"
 
diff --git a/Engine/Source/Editor/DetailCustomizations/Private/DirectoryPathStructCustomization.h b/Engine/Source/Editor/DetailCustomizations/Private/DirectoryPathStructCustomization.h
--- a/Engine/Source/Editor/DetailCustomizations/Private/DirectoryPathStructCustomization.h
+++ b/Engine/Source/Editor/DetailCustomizations/Private/DirectoryPathStructCustomization.h
@@ -9,6 +9,7 @@
 
 class SButton;
 class SComboButton;
+class IMenu;
 
 class FDirectoryPathStructCustomization : public IPropertyTypeCustomization
 {
diff --git a/Engine/Source/Editor/DetailCustomizations/Private/StaticMeshComponentDetails.cpp b/Engine/Source/Editor/DetailCustomizations/Private/StaticMeshComponentDetails.cpp
--- a/Engine/Source/Editor/DetailCustomizations/Private/StaticMeshComponentDetails.cpp
+++ b/Engine/Source/Editor/DetailCustomizations/Private/StaticMeshComponentDetails.cpp
@@ -9,6 +9,7 @@
 #include "IDetailPropertyRow.h"
 #include "DetailCategoryBuilder.h"
 #include "DetailWidgetRow.h"
+#include "Engine/StaticMesh.h"
 
 #define LOCTEXT_NAMESPACE "StaticMeshComponentDetails"
 
diff --git a/Engine/Source/Editor/DetailCustomizations/Private/Vector4StructCustomization.cpp b/Engine/Source/Editor/DetailCustomizations/Private/Vector4StructCustomization.cpp
--- a/Engine/Source/Editor/DetailCustomizations/Private/Vector4StructCustomization.cpp
+++ b/Engine/Source/Editor/DetailCustomizations/Private/Vector4StructCustomization.cpp
@@ -6,6 +6,7 @@
 #include "IDetailChildrenBuilder.h"
 #include "DetailLayoutBuilder.h"
 #include "DetailWidgetRow.h"
+#include "UObject/UnrealType.h"
 
 TSharedRef<IPropertyTypeCustomization> FVector4StructCustomization::MakeInstance()
 {
diff --git a/Engine/Source/Editor/Localization/Private/LocalizationConfigurationScript.cpp b/Engine/Source/Editor/Localization/Private/LocalizationConfigurationScript.cpp
--- a/Engine/Source/Editor/Localization/Private/LocalizationConfigurationScript.cpp
+++ b/Engine/Source/Editor/Localization/Private/LocalizationConfigurationScript.cpp
@@ -3,6 +3,7 @@
 #include "LocalizationConfigurationScript.h"
 #include "LocalizationTargetTypes.h"
 #include "LocalizationSettings.h"
+#include "UObject/Package.h"
 
 namespace
 {
diff --git a/Engine/Source/Editor/Persona/Private/SAnimationBlendSpaceGridWidget.h b/Engine/Source/Editor/Persona/Private/SAnimationBlendSpaceGridWidget.h
--- a/Engine/Source/Editor/Persona/Private/SAnimationBlendSpaceGridWidget.h
+++ b/Engine/Source/Editor/Persona/Private/SAnimationBlendSpaceGridWidget.h
@@ -13,6 +13,7 @@
 #include "Fonts/SlateFontInfo.h"
 #include "Widgets/SToolTip.h"
 #include "Misc/NotifyHook.h"
+#include "Animation/AnimationAsset.h"
 
 class FPaintArgs;
 class FSlateWindowElementList;
diff --git a/Engine/Source/Editor/PropertyEditor/Private/UserInterface/PropertyEditor/SPropertyEditorCombo.h b/Engine/Source/Editor/PropertyEditor/Private/UserInterface/PropertyEditor/SPropertyEditorCombo.h
--- a/Engine/Source/Editor/PropertyEditor/Private/UserInterface/PropertyEditor/SPropertyEditorCombo.h
+++ b/Engine/Source/Editor/PropertyEditor/Private/UserInterface/PropertyEditor/SPropertyEditorCombo.h
@@ -27,7 +27,7 @@ public:
 
 	void GetDesiredWidth( float& OutMinDesiredWidth, float& OutMaxDesiredWidth );
 private:
-	void GenerateComboBoxStrings( TArray< TSharedPtr<FString> >& OutComboBoxStrings, TArray<TSharedPtr<SToolTip>>& OutToolTips, TArray<bool>& OutRestrictedItems );
+	void GenerateComboBoxStrings( TArray< TSharedPtr<FString> >& OutComboBoxStrings, TArray<TSharedPtr<class SToolTip>>& OutToolTips, TArray<bool>& OutRestrictedItems );
 	void OnComboSelectionChanged( TSharedPtr<FString> NewValue, ESelectInfo::Type SelectInfo );
 	void OnComboOpening();
 
diff --git a/Engine/Source/Editor/PropertyEditor/Private/UserInterface/PropertyEditor/SPropertyEditorEditInline.cpp b/Engine/Source/Editor/PropertyEditor/Private/UserInterface/PropertyEditor/SPropertyEditorEditInline.cpp
--- a/Engine/Source/Editor/PropertyEditor/Private/UserInterface/PropertyEditor/SPropertyEditorEditInline.cpp
+++ b/Engine/Source/Editor/PropertyEditor/Private/UserInterface/PropertyEditor/SPropertyEditorEditInline.cpp
@@ -10,7 +10,7 @@
 #include "ClassViewerFilter.h"
 #include "Styling/SlateIconFinder.h"
 #include "ConstructorHelpers.h"
-#include "Editor/EditorEngine.h"
+#include "Editor.h"
 
 class FPropertyEditorInlineClassFilter : public IClassViewerFilter
 {
diff --git a/Engine/Source/Editor/SkeletonEditor/Private/BoneProxy.h b/Engine/Source/Editor/SkeletonEditor/Private/BoneProxy.h
--- a/Engine/Source/Editor/SkeletonEditor/Private/BoneProxy.h
+++ b/Engine/Source/Editor/SkeletonEditor/Private/BoneProxy.h
@@ -2,9 +2,14 @@
 
 #pragma once
 
+#include "CoreMinimal.h"
+#include "UObject/ObjectMacros.h"
 #include "TickableEditorObject.h"
+#include "UObject/Object.h"
+#include "UObject/WeakObjectPtr.h"
 #include "BoneProxy.generated.h"
 
+
 class UDebugSkelMeshComponent;
 
 /** Proxy object used to display/edit bone transforms */
diff --git a/Engine/Source/Editor/SkeletonEditor/Private/BoneProxyDetailsCustomization.cpp b/Engine/Source/Editor/SkeletonEditor/Private/BoneProxyDetailsCustomization.cpp
--- a/Engine/Source/Editor/SkeletonEditor/Private/BoneProxyDetailsCustomization.cpp
+++ b/Engine/Source/Editor/SkeletonEditor/Private/BoneProxyDetailsCustomization.cpp
@@ -10,6 +10,9 @@
 #include "IDetailPropertyRow.h"
 #include "Animation/DebugSkelMeshComponent.h"
 #include "ScopedTransaction.h"
+#include "MultiBoxBuilder.h"
+#include "SComboButton.h"
+#include "SBox.h"
 
 #define LOCTEXT_NAMESPACE "FBoneProxyDetailsCustomization"
 
diff --git a/Engine/Source/Editor/SourceControlWindows/Private/SSourceControlRevert.cpp b/Engine/Source/Editor/SourceControlWindows/Private/SSourceControlRevert.cpp
--- a/Engine/Source/Editor/SourceControlWindows/Private/SSourceControlRevert.cpp
+++ b/Engine/Source/Editor/SourceControlWindows/Private/SSourceControlRevert.cpp
@@ -29,9 +29,9 @@
 #include "ISourceControlModule.h"
 #include "PackageTools.h"
 #include "Settings/EditorExperimentalSettings.h"
-
 #include "Framework/Notifications/NotificationManager.h"
 #include "Widgets/Notifications/SNotificationList.h"
+#include "Linker.h"
 
 #define LOCTEXT_NAMESPACE "SSourceControlRevert"
 
diff --git a/Engine/Source/Editor/StaticMeshEditor/Private/StaticMeshEditorTools.cpp b/Engine/Source/Editor/StaticMeshEditor/Private/StaticMeshEditorTools.cpp
--- a/Engine/Source/Editor/StaticMeshEditor/Private/StaticMeshEditorTools.cpp
+++ b/Engine/Source/Editor/StaticMeshEditor/Private/StaticMeshEditorTools.cpp
@@ -32,6 +32,7 @@
 #include "Widgets/Input/STextComboBox.h"
 #include "ScopedTransaction.h"
 #include "JsonObjectConverter.h"
+#include "Engine/SkeletalMesh.h"
 
 const float MaxHullAccuracy = 1.f;
 const float MinHullAccuracy = 0.f;
diff --git a/Engine/Source/Editor/UnrealEd/Public/EditorReimportHandler.h b/Engine/Source/Editor/UnrealEd/Public/EditorReimportHandler.h
--- a/Engine/Source/Editor/UnrealEd/Public/EditorReimportHandler.h
+++ b/Engine/Source/Editor/UnrealEd/Public/EditorReimportHandler.h
@@ -4,6 +4,7 @@
 #pragma once
 
 #include "CoreMinimal.h"
+#include "UObject/GCObject.h"
 #include "Containers/ArrayView.h"
 
 class FReimportHandler;
diff --git a/Engine/Source/Runtime/AppFramework/Private/Widgets/Colors/SColorGradingPicker.cpp b/Engine/Source/Runtime/AppFramework/Private/Widgets/Colors/SColorGradingPicker.cpp
--- a/Engine/Source/Runtime/AppFramework/Private/Widgets/Colors/SColorGradingPicker.cpp
+++ b/Engine/Source/Runtime/AppFramework/Private/Widgets/Colors/SColorGradingPicker.cpp
@@ -2,6 +2,7 @@
 #include "SColorGradingPicker.h"
 #include "SColorGradingWheel.h"
 #include "SNumericEntryBox.h"
+#include "SBox.h"
 
 
 #define LOCTEXT_NAMESPACE "ColorGradingWheel"
diff --git a/Engine/Source/Runtime/AppFramework/Public/Widgets/Colors/SColorGradingPicker.h b/Engine/Source/Runtime/AppFramework/Public/Widgets/Colors/SColorGradingPicker.h
--- a/Engine/Source/Runtime/AppFramework/Public/Widgets/Colors/SColorGradingPicker.h
+++ b/Engine/Source/Runtime/AppFramework/Public/Widgets/Colors/SColorGradingPicker.h
@@ -8,6 +8,7 @@
 #include "Widgets/DeclarativeSyntaxSupport.h"
 #include "Widgets/SLeafWidget.h"
 #include "Framework/SlateDelegates.h"
+#include "SCompoundWidget.h"
 
 /** Callback to get the current FVector4 value */
 DECLARE_DELEGATE_OneParam(FOnGetCurrentVector4Value, FVector4&)
diff --git a/Engine/Source/Runtime/UMG/Public/Blueprint/WidgetLayoutLibrary.h b/Engine/Source/Runtime/UMG/Public/Blueprint/WidgetLayoutLibrary.h
--- a/Engine/Source/Runtime/UMG/Public/Blueprint/WidgetLayoutLibrary.h
+++ b/Engine/Source/Runtime/UMG/Public/Blueprint/WidgetLayoutLibrary.h
@@ -51,7 +51,7 @@ public:
 	/**
 	 * Gets the current DPI Scale being applied to the viewport and all the Widgets.
 	 */
-	static float GetViewportScale(UGameViewportClient* ViewportClient);
+	static float GetViewportScale(class UGameViewportClient* ViewportClient);
 
 	/**
 	 * Gets the size of the game viewport.
diff --git a/Engine/Source/Runtime/UMG/Public/Components/BackgroundBlurSlot.h b/Engine/Source/Runtime/UMG/Public/Components/BackgroundBlurSlot.h
--- a/Engine/Source/Runtime/UMG/Public/Components/BackgroundBlurSlot.h
+++ b/Engine/Source/Runtime/UMG/Public/Components/BackgroundBlurSlot.h
@@ -2,8 +2,10 @@
 
 #pragma once
 
+#include "UObject/ScriptMacros.h"
 #include "SlateWrapperTypes.h"
 #include "PanelSlot.h"
+#include "Layout/Margin.h"
 #include "BackgroundBlurSlot.generated.h"
 
 
