Fix OnLocationChanged broadcast for Android to use game thread properly

#android

diff --git a/Engine/Plugins/Runtime/LocationServicesAndroidImpl/Source/LocationServicesAndroidImpl/Private/LocationServicesAndroidImpl.cpp b/Engine/Plugins/Runtime/LocationServicesAndroidImpl/Source/LocationServicesAndroidImpl/Private/LocationServicesAndroidImpl.cpp
--- a/Engine/Plugins/Runtime/LocationServicesAndroidImpl/Source/LocationServicesAndroidImpl/Private/LocationServicesAndroidImpl.cpp
+++ b/Engine/Plugins/Runtime/LocationServicesAndroidImpl/Source/LocationServicesAndroidImpl/Private/LocationServicesAndroidImpl.cpp
@@ -111,5 +111,11 @@ JNI_METHOD void Java_com_epicgames_unreal_GameActivity_nativeHandleLocationChang
 	LocationData.VerticalAccuracy = 0.0f;
 	LocationData.Altitude = (float)altitude;
 
-	ULocationServices::GetLocationServicesImpl()->OnLocationChanged.Broadcast(LocationData);
+	if (FTaskGraphInterface::IsRunning())
+	{
+		FGraphEventRef BroadcastLocationTask = FFunctionGraphTask::CreateAndDispatchWhenReady([LocationData]()
+		{
+			ULocationServices::GetLocationServicesImpl()->OnLocationChanged.Broadcast(LocationData);
+		}, TStatId(), NULL, ENamedThreads::GameThread);
+	}
 }
